Change some more modules set as "Developer" type to "UncookedOnly".

#rnx

diff --git a/Engine/Plugins/Developer/CLionSourceCodeAccess/CLionSourceCodeAccess.uplugin b/Engine/Plugins/Developer/CLionSourceCodeAccess/CLionSourceCodeAccess.uplugin
--- a/Engine/Plugins/Developer/CLionSourceCodeAccess/CLionSourceCodeAccess.uplugin
+++ b/Engine/Plugins/Developer/CLionSourceCodeAccess/CLionSourceCodeAccess.uplugin
@@ -18,7 +18,7 @@
 	[
 		{
 			"Name" : "CLionSourceCodeAccess",
-			"Type" : "Developer",
+			"Type" : "UncookedOnly",
 			"LoadingPhase" : "Default",
 			"WhitelistPlatforms" :
 			[
diff --git a/Engine/Plugins/Developer/CodeLiteSourceCodeAccess/CodeLiteSourceCodeAccess.uplugin b/Engine/Plugins/Developer/CodeLiteSourceCodeAccess/CodeLiteSourceCodeAccess.uplugin
--- a/Engine/Plugins/Developer/CodeLiteSourceCodeAccess/CodeLiteSourceCodeAccess.uplugin
+++ b/Engine/Plugins/Developer/CodeLiteSourceCodeAccess/CodeLiteSourceCodeAccess.uplugin
@@ -18,7 +18,7 @@
 	[
 		{
 			"Name" : "CodeLiteSourceCodeAccess",
-			"Type" : "Developer",
+			"Type" : "UncookedOnly",
 			"LoadingPhase" : "Default",
 			"WhitelistPlatforms" :
 			[
diff --git a/Engine/Plugins/Developer/KDevelopSourceCodeAccess/KDevelopSourceCodeAccess.uplugin b/Engine/Plugins/Developer/KDevelopSourceCodeAccess/KDevelopSourceCodeAccess.uplugin
--- a/Engine/Plugins/Developer/KDevelopSourceCodeAccess/KDevelopSourceCodeAccess.uplugin
+++ b/Engine/Plugins/Developer/KDevelopSourceCodeAccess/KDevelopSourceCodeAccess.uplugin
@@ -18,7 +18,7 @@
 	[
 		{
 			"Name" : "KDevelopSourceCodeAccess",
-			"Type" : "Developer",
+			"Type" : "UncookedOnly",
 			"LoadingPhase" : "Default",
 			"WhitelistPlatforms" :
 			[
diff --git a/Engine/Plugins/Developer/NullSourceCodeAccess/NullSourceCodeAccess.uplugin b/Engine/Plugins/Developer/NullSourceCodeAccess/NullSourceCodeAccess.uplugin
--- a/Engine/Plugins/Developer/NullSourceCodeAccess/NullSourceCodeAccess.uplugin
+++ b/Engine/Plugins/Developer/NullSourceCodeAccess/NullSourceCodeAccess.uplugin
@@ -18,7 +18,7 @@
 	[
 		{
 			"Name" : "NullSourceCodeAccess",
-			"Type" : "Developer",
+			"Type" : "UncookedOnly",
 			"LoadingPhase" : "PreDefault",
 			"WhitelistPlatforms" :
 			[
diff --git a/Engine/Plugins/Developer/VisualStudioCodeSourceCodeAccess/VisualStudioCodeSourceCodeAccess.uplugin b/Engine/Plugins/Developer/VisualStudioCodeSourceCodeAccess/VisualStudioCodeSourceCodeAccess.uplugin
--- a/Engine/Plugins/Developer/VisualStudioCodeSourceCodeAccess/VisualStudioCodeSourceCodeAccess.uplugin
+++ b/Engine/Plugins/Developer/VisualStudioCodeSourceCodeAccess/VisualStudioCodeSourceCodeAccess.uplugin
@@ -18,7 +18,7 @@
 	[
 		{
 			"Name" : "VisualStudioCodeSourceCodeAccess",
-			"Type" : "Developer",
+			"Type" : "UncookedOnly",
 			"LoadingPhase" : "Default",
 			"WhitelistPlatforms" :
 			[
diff --git a/Engine/Plugins/Developer/VisualStudioSourceCodeAccess/VisualStudioSourceCodeAccess.uplugin b/Engine/Plugins/Developer/VisualStudioSourceCodeAccess/VisualStudioSourceCodeAccess.uplugin
--- a/Engine/Plugins/Developer/VisualStudioSourceCodeAccess/VisualStudioSourceCodeAccess.uplugin
+++ b/Engine/Plugins/Developer/VisualStudioSourceCodeAccess/VisualStudioSourceCodeAccess.uplugin
@@ -19,7 +19,7 @@
 	[
 		{
 			"Name" : "VisualStudioSourceCodeAccess",
-			"Type" : "Developer",
+			"Type" : "UncookedOnly",
 			"LoadingPhase" : "Default",
 			"WhitelistPlatforms" :
 			[
diff --git a/Engine/Plugins/Developer/XCodeSourceCodeAccess/XCodeSourceCodeAccess.uplugin b/Engine/Plugins/Developer/XCodeSourceCodeAccess/XCodeSourceCodeAccess.uplugin
--- a/Engine/Plugins/Developer/XCodeSourceCodeAccess/XCodeSourceCodeAccess.uplugin
+++ b/Engine/Plugins/Developer/XCodeSourceCodeAccess/XCodeSourceCodeAccess.uplugin
@@ -19,7 +19,7 @@
 	[
 		{
 			"Name" : "XCodeSourceCodeAccess",
-			"Type" : "Developer",
+			"Type" : "UncookedOnly",
 			"LoadingPhase" : "Default",
 			"WhitelistPlatforms" :
 			[
