Fix Visual Studio Code file open to better create command line for Linux compatibility

#rnx

diff --git a/Engine/Plugins/Developer/VisualStudioCodeSourceCodeAccess/Source/VisualStudioCodeSourceCodeAccess/Private/VisualStudioCodeSourceCodeAccessor.cpp b/Engine/Plugins/Developer/VisualStudioCodeSourceCodeAccess/Source/VisualStudioCodeSourceCodeAccess/Private/VisualStudioCodeSourceCodeAccessor.cpp
--- a/Engine/Plugins/Developer/VisualStudioCodeSourceCodeAccess/Source/VisualStudioCodeSourceCodeAccess/Private/VisualStudioCodeSourceCodeAccessor.cpp
+++ b/Engine/Plugins/Developer/VisualStudioCodeSourceCodeAccess/Source/VisualStudioCodeSourceCodeAccess/Private/VisualStudioCodeSourceCodeAccessor.cpp
@@ -86,7 +86,7 @@ void FVisualStudioCodeSourceCodeAccessor::RefreshAvailability()
 	FString OutURL;
 	int32 ReturnCode = -1;
 
-	FPlatformProcess::ExecProcess(TEXT("/bin/bash"), TEXT("-c \"type -p code\""), &ReturnCode, &OutURL, nullptr);
+	FPlatformProcess::ExecProcess(TEXT("/usr/bin/env"), TEXT("bash -c \"type -p code\""), &ReturnCode, &OutURL, nullptr);
 	if (ReturnCode == 0)
 	{
 		Location.URL = OutURL.TrimStartAndEnd();
@@ -149,7 +149,8 @@ bool FVisualStudioCodeSourceCodeAccessor::OpenFileAtLine(const FString& FullPath
 		FString SolutionDir = GetSolutionPath();
 		TArray<FString> Args;
 		Args.Add(MakePath(SolutionDir));
-		Args.Add(TEXT("-g ") + MakePath(FullPath) + FString::Printf(TEXT(":%d:%d"), LineNumber, ColumnNumber));
+		Args.Add(TEXT("-g"));
+		Args.Add(MakePath(FullPath + FString::Printf(TEXT(":%d:%d"), LineNumber, ColumnNumber)));
 		return Launch(Args);
 	}
 
@@ -234,7 +235,15 @@ bool FVisualStudioCodeSourceCodeAccessor::Launch(const TArray<FString>& InArgs)
 
 		uint32 ProcessID;
 		FProcHandle hProcess = FPlatformProcess::CreateProc(*Location.URL, *ArgsString, true, false, false, &ProcessID, 0, nullptr, nullptr, nullptr);
-		return hProcess.IsValid();
+		if (hProcess.IsValid())
+		{
+			FPlatformProcess::CloseProc(hProcess);
+			return true;
+		}
+		else
+		{
+			return false;
+		}
 	}
 	
 	return false;
