Fix VisualStudioCode Accessor to find the code-workspace like Visual Studio

#rnx

diff --git a/Engine/Plugins/Developer/VisualStudioCodeSourceCodeAccess/Source/VisualStudioCodeSourceCodeAccess/Private/VisualStudioCodeSourceCodeAccessor.cpp b/Engine/Plugins/Developer/VisualStudioCodeSourceCodeAccess/Source/VisualStudioCodeSourceCodeAccess/Private/VisualStudioCodeSourceCodeAccessor.cpp
--- a/Engine/Plugins/Developer/VisualStudioCodeSourceCodeAccess/Source/VisualStudioCodeSourceCodeAccess/Private/VisualStudioCodeSourceCodeAccessor.cpp
+++ b/Engine/Plugins/Developer/VisualStudioCodeSourceCodeAccess/Source/VisualStudioCodeSourceCodeAccess/Private/VisualStudioCodeSourceCodeAccessor.cpp
@@ -6,6 +6,9 @@
 #include "Misc/Paths.h"
 #include "Misc/UProjectInfo.h"
 #include "Misc/App.h"
+#include "Misc/FileHelper.h"
+#include "ProjectDescriptor.h"
+#include "Interfaces/IProjectManager.h"
 
 #if PLATFORM_WINDOWS
 #include "Internationalization/Regex.h"
@@ -23,7 +26,7 @@ namespace
 
 static FString MakePath(const FString& InPath)
 {
-	return TEXT("\"") + InPath + TEXT("\""); 
+	return TEXT("\"") + InPath.ReplaceQuotesWithEscapedQuotes() + TEXT("\""); 
 }
 
 FString FVisualStudioCodeSourceCodeAccessor::GetSolutionPath() const
@@ -32,19 +35,28 @@ FString FVisualStudioCodeSourceCodeAccessor::GetSolutionPath() const
 
 	if (IsInGameThread())
 	{
-		CachedSolutionPath = FPaths::ProjectDir();
-
-		if (!FUProjectDictionary::GetDefault().IsForeignProject(CachedSolutionPath))
-		{
-			CachedSolutionPath = FPaths::Combine(FPaths::RootDir(), FString("UE5") + GVSCodeWorkspaceExtension);
-		}
-		else
+		FString PrimaryProjectPath;
+		if (!FFileHelper::LoadFileToString(PrimaryProjectPath, *(FPaths::EngineIntermediateDir() / TEXT("ProjectFiles") / TEXT("PrimaryProjectPath.txt"))) || 
+			(PrimaryProjectPath.TrimStartAndEndInline(), PrimaryProjectPath.IsEmpty()))
 		{
-			FString BaseName = FApp::HasProjectName() ? FApp::GetProjectName() : FPaths::GetBaseFilename(CachedSolutionPath);
-			CachedSolutionPath = FPaths::Combine(CachedSolutionPath, BaseName + GVSCodeWorkspaceExtension);
-		}
+			const FProjectDescriptor* CurrentProject = IProjectManager::Get().GetCurrentProject();
+			const FString ProjectDir = FPaths::ConvertRelativePathToFull(FPaths::ProjectDir());
+
+			if ((CurrentProject == nullptr || CurrentProject->Modules.Num() == 0) || !FUProjectDictionary::GetDefault().IsForeignProject(ProjectDir))
+			{
+				PrimaryProjectPath = FPaths::RootDir() / TEXT("UE5");
+			}
+			else
+			{
+				const FString BaseName = FApp::HasProjectName() ? FApp::GetProjectName() : FPaths::GetBaseFilename(ProjectDir);
+				PrimaryProjectPath = ProjectDir / BaseName;
+			}
+		}	
+
+		CachedSolutionPath = PrimaryProjectPath + GVSCodeWorkspaceExtension;
+		FPaths::NormalizeFilename(CachedSolutionPath);			
 	}
-
+	
 	return CachedSolutionPath;
 }
 
diff --git a/Engine/Plugins/Developer/VisualStudioCodeSourceCodeAccess/Source/VisualStudioCodeSourceCodeAccess/VisualStudioCodeSourceCodeAccess.Build.cs b/Engine/Plugins/Developer/VisualStudioCodeSourceCodeAccess/Source/VisualStudioCodeSourceCodeAccess/VisualStudioCodeSourceCodeAccess.Build.cs
--- a/Engine/Plugins/Developer/VisualStudioCodeSourceCodeAccess/Source/VisualStudioCodeSourceCodeAccess/VisualStudioCodeSourceCodeAccess.Build.cs
+++ b/Engine/Plugins/Developer/VisualStudioCodeSourceCodeAccess/Source/VisualStudioCodeSourceCodeAccess/VisualStudioCodeSourceCodeAccess.Build.cs
@@ -13,7 +13,8 @@ namespace UnrealBuildTool.Rules
 				{
 					"Core",
 					"SourceCodeAccess",
-					"DesktopPlatform"
+					"DesktopPlatform",
+					"Projects",
 				}
 			);
 
