Add missing HideWindowsPlatformTypes

#Synced-CL Unknown

diff --git a/Engine/Plugins/Developer/CLionSourceCodeAccess/Source/CLionSourceCodeAccess/Private/CLionSourceCodeAccessor.cpp b/Engine/Plugins/Developer/CLionSourceCodeAccess/Source/CLionSourceCodeAccess/Private/CLionSourceCodeAccessor.cpp
--- a/Engine/Plugins/Developer/CLionSourceCodeAccess/Source/CLionSourceCodeAccess/Private/CLionSourceCodeAccessor.cpp
+++ b/Engine/Plugins/Developer/CLionSourceCodeAccess/Source/CLionSourceCodeAccess/Private/CLionSourceCodeAccessor.cpp
@@ -339,3 +339,7 @@ FString FCLionSourceCodeAccessor::FindExecutablePath()
 }
 
 #undef LOCTEXT_NAMESPACE
+
+#if PLATFORM_WINDOWS
+#include "Windows/HideWindowsPlatformTypes.h"
+#endif
diff --git a/Engine/Plugins/Developer/VisualStudioCodeSourceCodeAccess/Source/VisualStudioCodeSourceCodeAccess/Private/VisualStudioCodeSourceCodeAccessor.cpp b/Engine/Plugins/Developer/VisualStudioCodeSourceCodeAccess/Source/VisualStudioCodeSourceCodeAccess/Private/VisualStudioCodeSourceCodeAccessor.cpp
--- a/Engine/Plugins/Developer/VisualStudioCodeSourceCodeAccess/Source/VisualStudioCodeSourceCodeAccess/Private/VisualStudioCodeSourceCodeAccessor.cpp
+++ b/Engine/Plugins/Developer/VisualStudioCodeSourceCodeAccess/Source/VisualStudioCodeSourceCodeAccess/Private/VisualStudioCodeSourceCodeAccessor.cpp
@@ -262,3 +262,7 @@ bool FVisualStudioCodeSourceCodeAccessor::Launch(const TArray<FString>& InArgs)
 }
 
 #undef LOCTEXT_NAMESPACE
+
+#if PLATFORM_WINDOWS
+#include "Windows/HideWindowsPlatformTypes.h"
+#endif
